Replace LogUE4ML by LogUnrealEditor4ML

diff --git a/Engine/Plugins/AI/UE4ML/Source/UE4ML/Private/4MLLibrarian.cpp b/Engine/Plugins/AI/UE4ML/Source/UE4ML/Private/4MLLibrarian.cpp
--- a/Engine/Plugins/AI/UE4ML/Source/UE4ML/Private/4MLLibrarian.cpp
+++ b/Engine/Plugins/AI/UE4ML/Source/UE4ML/Private/4MLLibrarian.cpp
@@ -20,7 +20,7 @@ namespace FLibrarianHelper
 
 		if (Class->HasAnyClassFlags(CLASS_Abstract | CLASS_Deprecated | CLASS_NewerVersionExists))
 		{
-			UE_LOG(LogUE4ML, Log, TEXT("Librarian: skipping class %s registration due to it being %s%s%s")
+			UE_LOG(LogUnrealEditor4ML, Log, TEXT("Librarian: skipping class %s registration due to it being %s%s%s")
 				, *Class->GetName()
 				, Class->HasAnyClassFlags(CLASS_Abstract) ? TEXT("Abstract, ") : TEXT("")
 				, Class->HasAnyClassFlags(CLASS_Deprecated) ? TEXT("Deprecated, ") : TEXT("")
diff --git a/Engine/Plugins/AI/UE4ML/Source/UE4ML/Private/4MLManager.cpp b/Engine/Plugins/AI/UE4ML/Source/UE4ML/Private/4MLManager.cpp
--- a/Engine/Plugins/AI/UE4ML/Source/UE4ML/Private/4MLManager.cpp
+++ b/Engine/Plugins/AI/UE4ML/Source/UE4ML/Private/4MLManager.cpp
@@ -70,7 +70,7 @@ namespace
 				? SettingsManagerClass.Get()
 				: U4MLManager::StaticClass();
 
-			UE_LOG(LogUE4ML, Log, TEXT("Creating UE4ML manager of class %s"), *GetNameSafe(Class));
+			UE_LOG(LogUnrealEditor4ML, Log, TEXT("Creating UE4ML manager of class %s"), *GetNameSafe(Class));
 
 			U4MLManager* ManagerInstance = NewObject<U4MLManager>(GEngine, Class);
 			check(ManagerInstance);
@@ -143,7 +143,7 @@ void U4MLManager::StopServer()
 {
 	if (RPCServerInstance)
 	{
-		UE_LOG(LogUE4ML, Log, TEXT("Stopping RPC server."));
+		UE_LOG(LogUnrealEditor4ML, Log, TEXT("Stopping RPC server."));
 		RPCServerInstance->stop();
 		delete RPCServerInstance;
 		RPCServerInstance = nullptr;
@@ -179,7 +179,7 @@ void U4MLManager::StartServer(uint16 Port, EUE4MLServerMode InMode, uint16 Serve
 		}
 	}
 
-	UE_LOG(LogUE4ML, Log, TEXT("Starting RPC server on port %d."), Port);
+	UE_LOG(LogUnrealEditor4ML, Log, TEXT("Starting RPC server on port %d."), Port);
 	RPCServerInstance = new FRPCServer(Port);
 	CurrentPort = Port;
 	check(RPCServerInstance);
@@ -364,7 +364,7 @@ void U4MLManager::OnPostWorldInit(UWorld* World, const UWorld::InitializationVal
 
 					if (PlayNetMode != PIE_Standalone)
 					{
-						UE_LOG(LogUE4ML, Log, TEXT("Ignoring %s due to net mode != PIE_Standalone")
+						UE_LOG(LogUnrealEditor4ML, Log, TEXT("Ignoring %s due to net mode != PIE_Standalone")
 							, *World->GetName());
 						//UUE4MLSettings::
 						return;
diff --git a/Engine/Plugins/AI/UE4ML/Source/UE4ML/Private/4MLManager_Client.cpp b/Engine/Plugins/AI/UE4ML/Source/UE4ML/Private/4MLManager_Client.cpp
--- a/Engine/Plugins/AI/UE4ML/Source/UE4ML/Private/4MLManager_Client.cpp
+++ b/Engine/Plugins/AI/UE4ML/Source/UE4ML/Private/4MLManager_Client.cpp
@@ -94,7 +94,7 @@ void U4MLManager::AddCommonFunctions(FRPCServer& Server)
 
 void U4MLManager::ConfigureAsClient(FRPCServer& Server)
 {
-	UE_LOG(LogUE4ML, Log, TEXT("\tconfiguring as client"));
+	UE_LOG(LogUnrealEditor4ML, Log, TEXT("\tconfiguring as client"));
 
 	AddCommonFunctions(Server);
 
diff --git a/Engine/Plugins/AI/UE4ML/Source/UE4ML/Private/4MLManager_Server.cpp b/Engine/Plugins/AI/UE4ML/Source/UE4ML/Private/4MLManager_Server.cpp
--- a/Engine/Plugins/AI/UE4ML/Source/UE4ML/Private/4MLManager_Server.cpp
+++ b/Engine/Plugins/AI/UE4ML/Source/UE4ML/Private/4MLManager_Server.cpp
@@ -22,7 +22,7 @@
 
 void U4MLManager::ConfigureAsServer(FRPCServer& Server)
 {
-	UE_LOG(LogUE4ML, Log, TEXT("\tconfiguring as server"));
+	UE_LOG(LogUnrealEditor4ML, Log, TEXT("\tconfiguring as server"));
 
 	AddCommonFunctions(Server);
 
diff --git a/Engine/Plugins/AI/UE4ML/Source/UE4ML/Private/4MLSession.cpp b/Engine/Plugins/AI/UE4ML/Source/UE4ML/Private/4MLSession.cpp
--- a/Engine/Plugins/AI/UE4ML/Source/UE4ML/Private/4MLSession.cpp
+++ b/Engine/Plugins/AI/UE4ML/Source/UE4ML/Private/4MLSession.cpp
@@ -374,7 +374,7 @@ U4MLAgent* U4MLSession::GetAgent(F4ML::FAgentID AgentID)
 {
 	if (Agents.IsValidIndex(AgentID) == false)
 	{
-		UE_LOG(LogUE4ML, Warning, TEXT("U4MLSession::GetAgent: Invalid AgentID. Failing"));
+		UE_LOG(LogUnrealEditor4ML, Warning, TEXT("U4MLSession::GetAgent: Invalid AgentID. Failing"));
 		return nullptr;
 	}
 
@@ -467,7 +467,7 @@ bool U4MLSession::RequestAvatarForAgent(U4MLAgent& Agent, UWorld* InWorld, const
 	if (Agent.GetAvatar() != nullptr)
 	{
 		// skipping.
-		UE_LOG(LogUE4ML, Verbose, TEXT("U4MLSession::RequestAvatarForAgent called for agent [%s] while it still has an avatar [%s]. Call ClearAvatar first to null-out agent\'s avatar.")
+		UE_LOG(LogUnrealEditor4ML, Verbose, TEXT("U4MLSession::RequestAvatarForAgent called for agent [%s] while it still has an avatar [%s]. Call ClearAvatar first to null-out agent\'s avatar.")
 			, Agent.GetAgentID(), *GetNameSafe(Agent.GetAvatar()));
 		return false;
 	}
@@ -485,7 +485,7 @@ bool U4MLSession::RequestAvatarForAgent(U4MLAgent& Agent, UWorld* InWorld, const
 	
 	if (InWorld == nullptr)
 	{
-		UE_LOG(LogUE4ML, Warning, TEXT("U4MLSession::RequestAvatarForAgent called with InWorld and CachedWorld both being null. Auto-failure."));
+		UE_LOG(LogUnrealEditor4ML, Warning, TEXT("U4MLSession::RequestAvatarForAgent called with InWorld and CachedWorld both being null. Auto-failure."));
 		return false;
 	}
 
diff --git a/Engine/Plugins/AI/UE4ML/Source/UE4ML/Private/4MLTypes.cpp b/Engine/Plugins/AI/UE4ML/Source/UE4ML/Private/4MLTypes.cpp
--- a/Engine/Plugins/AI/UE4ML/Source/UE4ML/Private/4MLTypes.cpp
+++ b/Engine/Plugins/AI/UE4ML/Source/UE4ML/Private/4MLTypes.cpp
@@ -4,7 +4,7 @@
 #include "GameFramework/Controller.h"
 #include "GameFramework/Pawn.h"
 
-DEFINE_LOG_CATEGORY(LogUE4ML);
+DEFINE_LOG_CATEGORY(LogUnrealEditor4ML);
 
 namespace F4ML
 {
diff --git a/Engine/Plugins/AI/UE4ML/Source/UE4ML/Private/Agents/4MLAgent.cpp b/Engine/Plugins/AI/UE4ML/Source/UE4ML/Private/Agents/4MLAgent.cpp
--- a/Engine/Plugins/AI/UE4ML/Source/UE4ML/Private/Agents/4MLAgent.cpp
+++ b/Engine/Plugins/AI/UE4ML/Source/UE4ML/Private/Agents/4MLAgent.cpp
@@ -342,7 +342,7 @@ void U4MLAgent::SetAvatar(AActor* InAvatar)
 
 	if (InAvatar != nullptr && IsSuitableAvatar(*InAvatar) == false)
 	{
-		UE_LOG(LogUE4ML, Log, TEXT("SetAvatar was called for agent %u but %s is not a valid avatar (required avatar class %s)")
+		UE_LOG(LogUnrealEditor4ML, Log, TEXT("SetAvatar was called for agent %u but %s is not a valid avatar (required avatar class %s)")
 			, AgentID, *InAvatar->GetName(), *GetNameSafe(AgentConfig.AvatarClass));
 		return;
 	}
diff --git a/Engine/Plugins/AI/UE4ML/Source/UE4ML/Public/4MLTypes.h b/Engine/Plugins/AI/UE4ML/Source/UE4ML/Public/4MLTypes.h
--- a/Engine/Plugins/AI/UE4ML/Source/UE4ML/Public/4MLTypes.h
+++ b/Engine/Plugins/AI/UE4ML/Source/UE4ML/Public/4MLTypes.h
@@ -11,7 +11,7 @@
 class AActor;
 class AController;
 
-DECLARE_LOG_CATEGORY_EXTERN(LogUE4ML, Log, All);
+DECLARE_LOG_CATEGORY_EXTERN(LogUnrealEditor4ML, Log, All);
 
 #ifndef DEFINE_ENUM_TO_STRING
 #define DEFINE_ENUM_TO_STRING(EnumType) FString EnumToString(const EnumType Value) \
