Fixed a shadowing variable warning on clang.

#rnx

diff --git a/Engine/Plugins/AI/UE4ML/Source/UE4ML/Private/4MLManager_Client.cpp b/Engine/Plugins/AI/UE4ML/Source/UE4ML/Private/4MLManager_Client.cpp
--- a/Engine/Plugins/AI/UE4ML/Source/UE4ML/Private/4MLManager_Client.cpp
+++ b/Engine/Plugins/AI/UE4ML/Source/UE4ML/Private/4MLManager_Client.cpp
@@ -77,10 +77,10 @@ void U4MLManager::AddCommonFunctions(FRPCServer& Server)
 		}
 		else
 		{
-			U4MLSession& Session = GetSession();
+			U4MLSession& LocalSession = GetSession();
 			for (const F4ML::FAgentID& AgentID : AgentIDs)
 			{
-				U4MLAgent* Agent = Session.GetAgent(AgentID);
+				U4MLAgent* Agent = LocalSession.GetAgent(AgentID);
 				Results.push_back(Agent == nullptr || Agent->IsDone());
 			}
 		}
